[MAIN] Merging fix for XGE shader compiler (CL 3633490) from //UE4/Dev-Console/... to //UE4/Main/...
 - Missing VirtualizeDirectX command line switch was causing certain remote workers to fail.

diff --git a/Engine/Source/Developer/XGEController/Private/XGEControllerModule.cpp b/Engine/Source/Developer/XGEController/Private/XGEControllerModule.cpp
--- a/Engine/Source/Developer/XGEController/Private/XGEControllerModule.cpp
+++ b/Engine/Source/Developer/XGEController/Private/XGEControllerModule.cpp
@@ -371,7 +371,7 @@ void FXGEControllerModule::WriteOutThreadProc()
 
 		// To handle spaces in the engine path, we just pass the XGEController.exe filename to xgConsole,
 		// and set the working directory of xgConsole.exe to the engine binaries folder below.
-		FString XGConsoleArgs = FString::Printf(TEXT("/allowremote=\"%s\" /allowintercept=\"%s\" /title=\"Unreal Engine XGE Tasks\" /monitordirs=\"%s\" /command=\"%s -xgecontroller %s\""),
+		FString XGConsoleArgs = FString::Printf(TEXT("/VIRTUALIZEDIRECTX /allowremote=\"%s\" /allowintercept=\"%s\" /title=\"Unreal Engine XGE Tasks\" /monitordirs=\"%s\" /command=\"%s -xgecontroller %s\""),
 			XGE_INTERCEPT_EXE_NAMES,
 			XGE_CONTROL_WORKER_NAME,
 			*WorkingDirectory,
