[NO NON-PII COMMIT MESSAGE]

diff --git a/Engine/Source/Developer/XGEController/Private/XGEControllerModule.cpp b/Engine/Source/Developer/XGEController/Private/XGEControllerModule.cpp
--- a/Engine/Source/Developer/XGEController/Private/XGEControllerModule.cpp
+++ b/Engine/Source/Developer/XGEController/Private/XGEControllerModule.cpp
@@ -197,6 +197,19 @@ bool FXGEControllerModule::IsSupported()
 	// Check for a valid installation of Incredibuild by seeing if xgconsole.exe exists.
 	if (XGEControllerVariables::Enabled == 1)
 	{
+		// Try to read from the registry
+		FString RegistryPathString;
+
+		if (!FWindowsPlatformMisc::QueryRegKey(HKEY_LOCAL_MACHINE, TEXT("SOFTWARE\\Xoreax\\IncrediBuild\\Builder"), TEXT("Folder"), RegistryPathString))
+		{
+			FWindowsPlatformMisc::QueryRegKey(HKEY_LOCAL_MACHINE, TEXT("SOFTWARE\\WOW6432Node\\Xoreax\\IncrediBuild\\Builder"), TEXT("Folder"), RegistryPathString);
+		}
+
+		if (!RegistryPathString.IsEmpty())
+		{
+			RegistryPathString = FPaths::Combine(RegistryPathString, TEXT("xgConsole.exe"));
+		}
+
 		// Try to find xgConsole.exe from the PATH environment variable
 		FString PathString;
 		{
@@ -220,6 +233,7 @@ bool FXGEControllerModule::IsSupported()
 		// List of possible paths to xgconsole.exe
 		const FString Paths[] =
 		{
+			*RegistryPathString,
 			TEXT("C:\\Program Files\\Xoreax\\IncrediBuild\\xgConsole.exe"),
 			TEXT("C:\\Program Files (x86)\\Xoreax\\IncrediBuild\\xgConsole.exe"),
 			*PathString
