Add a dummy exported symbol to the XGEController module, to fix build errors due to missing .lib file when it's built with WITH_XGE_CONTROLLER = 0.

diff --git a/Engine/Source/Developer/XGEController/Private/XGEControllerModule.cpp b/Engine/Source/Developer/XGEController/Private/XGEControllerModule.cpp
--- a/Engine/Source/Developer/XGEController/Private/XGEControllerModule.cpp
+++ b/Engine/Source/Developer/XGEController/Private/XGEControllerModule.cpp
@@ -586,4 +586,9 @@ XGECONTROLLER_API IXGEController& IXGEController::Get()
 
 IMPLEMENT_MODULE(FXGEControllerModule, XGEController);
 
+#else
+
+// Workaround for module not having any exported symbols
+XGECONTROLLER_API int XgeControllerExportedSymbol = 0;
+
 #endif // WITH_XGE_CONTROLLER
