Modified lesser version check in FXGEControllerModule::IsSupported() to be a warning instead of an error. This error was failing a valid cook on a remote machine without XGE configured (and an old version installed).

(ushell-p4-cherrypick of 7972164 by graham.wihlidal)

diff --git a/Engine/Source/Developer/XGEController/Private/XGEControllerModule.cpp b/Engine/Source/Developer/XGEController/Private/XGEControllerModule.cpp
--- a/Engine/Source/Developer/XGEController/Private/XGEControllerModule.cpp
+++ b/Engine/Source/Developer/XGEController/Private/XGEControllerModule.cpp
@@ -295,7 +295,7 @@ bool FXGEControllerModule::IsSupported()
 			}
 			else if (Version < 1002867)
 			{
-				UE_LOG(LogXGEController, Error, TEXT("XGE version 8.01 (build 1867) or higher is required for XGE shader compilation with the interception interface."));
+				UE_LOG(LogXGEController, Warning, TEXT("XGE version 8.01 (build 1867) or higher is required for XGE shader compilation with the interception interface."));
 				XGEControllerVariables::Enabled = 0;
 			}
 		}
